Mover 2.0 tests: Updating plugin to have proper dependencies

diff --git a/Engine/Plugins/Experimental/MoverTests/MoverTests.uplugin b/Engine/Plugins/Experimental/MoverTests/MoverTests.uplugin
--- a/Engine/Plugins/Experimental/MoverTests/MoverTests.uplugin
+++ b/Engine/Plugins/Experimental/MoverTests/MoverTests.uplugin
@@ -25,6 +25,10 @@
 		{
 			"Name": "Mover",
 			"Enabled": true
+		},
+		{
+			"Name": "MoverExamples",
+			"Enabled": true
 		}
 	]
 }
\ No newline at end of file
diff --git a/Engine/Plugins/Experimental/MoverTests/Source/MoverTests/MoverTests.Build.cs b/Engine/Plugins/Experimental/MoverTests/Source/MoverTests/MoverTests.Build.cs
--- a/Engine/Plugins/Experimental/MoverTests/Source/MoverTests/MoverTests.Build.cs
+++ b/Engine/Plugins/Experimental/MoverTests/Source/MoverTests/MoverTests.Build.cs
@@ -38,6 +38,8 @@ public class MoverTests : ModuleRules
 				"Engine",
 				"Slate",
 				"SlateCore",
+				"Mover",
+				"MoverExamples"
 				// ... add private dependencies that you statically link with here ...	
 			}
 			);
