Mover 2.0 Tests: Adding a new test layered move that acts the same as Launch. Used for testing implementing layered moves.

diff --git a/Engine/Plugins/Experimental/MoverTests/Source/MoverTests/MoverTests.Build.cs b/Engine/Plugins/Experimental/MoverTests/Source/MoverTests/MoverTests.Build.cs
--- a/Engine/Plugins/Experimental/MoverTests/Source/MoverTests/MoverTests.Build.cs
+++ b/Engine/Plugins/Experimental/MoverTests/Source/MoverTests/MoverTests.Build.cs
@@ -26,6 +26,13 @@ public class MoverTests : ModuleRules
 			new string[]
 			{
 				"Core",
+				"CoreUObject",
+				"NetCore",
+				"Engine",
+				"Slate",
+				"SlateCore",
+				"Mover",
+				"MoverExamples"
 				// ... add other public dependencies that you statically link with here ...
 			}
 			);
@@ -34,12 +41,6 @@ public class MoverTests : ModuleRules
 		PrivateDependencyModuleNames.AddRange(
 			new string[]
 			{
-				"CoreUObject",
-				"Engine",
-				"Slate",
-				"SlateCore",
-				"Mover",
-				"MoverExamples"
 				// ... add private dependencies that you statically link with here ...	
 			}
 			);
diff --git a/Engine/Plugins/Experimental/MoverTests/Source/MoverTests/Private/TestCustomLayeredMoves.cpp b/Engine/Plugins/Experimental/MoverTests/Source/MoverTests/Private/TestCustomLayeredMoves.cpp
new file mode 100644
--- /dev/null
+++ b/Engine/Plugins/Experimental/MoverTests/Source/MoverTests/Private/TestCustomLayeredMoves.cpp
@@ -0,0 +1,68 @@
+﻿// Copyright Epic Games, Inc. All Rights Reserved.
+
+#include "TestCustomLayeredMoves.h"
+#include "MoverComponent.h"
+#include "MoverLog.h"
+
+FTestCustomLayeredMove::FTestCustomLayeredMove()
+{
+	DurationMs = 0.f;
+	MixMode = EMoveMixMode::OverrideVelocity;
+}
+
+bool FTestCustomLayeredMove::GenerateMove(const FMoverTickStartData& SimState, const FMoverTimeStep& TimeStep, const UMoverComponent* MoverComp, UMoverBlackboard* SimBlackboard, FProposedMove& OutProposedMove)
+{
+	OutProposedMove.MixMode = MixMode;
+	OutProposedMove.LinearVelocity = LaunchVelocity;
+	OutProposedMove.PreferredMode = ForceMovementMode;
+
+	return true;
+}
+
+void FTestCustomLayeredMove::OnStart(const UMoverComponent* MoverComp, UMoverBlackboard* SimBlackboard)
+{
+	UE_LOG(LogMover, Log, TEXT("Custom test layered move started!"));
+}
+
+void FTestCustomLayeredMove::OnEnd(const UMoverComponent* MoverComp, UMoverBlackboard* SimBlackboard, float CurrentSimTimeMs)
+{
+	UE_LOG(LogMover, Log, TEXT("Custom test layered move ended!"));
+}
+
+FLayeredMoveBase* FTestCustomLayeredMove::Clone() const
+{
+	FTestCustomLayeredMove* CopyPtr = new FTestCustomLayeredMove(*this);
+	return CopyPtr;
+}
+
+void FTestCustomLayeredMove::NetSerialize(FArchive& Ar)
+{
+	Super::NetSerialize(Ar);
+
+	SerializePackedVector<10, 16>(LaunchVelocity, Ar);
+
+	bool bUsingForcedMovementMode = !ForceMovementMode.IsNone();
+
+	Ar.SerializeBits(&bUsingForcedMovementMode, 1);
+
+	if (bUsingForcedMovementMode)
+	{
+		Ar << ForceMovementMode;
+	}
+
+}
+
+UScriptStruct* FTestCustomLayeredMove::GetScriptStruct() const
+{
+	return FTestCustomLayeredMove::StaticStruct();
+}
+
+FString FTestCustomLayeredMove::ToSimpleString() const
+{
+	return FString::Printf(TEXT("Custom Test Move"));
+}
+
+void FTestCustomLayeredMove::AddReferencedObjects(class FReferenceCollector& Collector)
+{
+	Super::AddReferencedObjects(Collector);
+}
\ No newline at end of file
diff --git a/Engine/Plugins/Experimental/MoverTests/Source/MoverTests/Public/TestCustomLayeredMoves.h b/Engine/Plugins/Experimental/MoverTests/Source/MoverTests/Public/TestCustomLayeredMoves.h
new file mode 100644
--- /dev/null
+++ b/Engine/Plugins/Experimental/MoverTests/Source/MoverTests/Public/TestCustomLayeredMoves.h
@@ -0,0 +1,43 @@
+﻿// Copyright Epic Games, Inc. All Rights Reserved.
+
+#pragma once
+
+#include "LayeredMove.h"
+#include "TestCustomLayeredMoves.generated.h"
+
+/**
+ * Custom Layered move for testing purposes - acts the same as the Launch layered move
+ */
+USTRUCT(BlueprintType)
+struct MOVERTESTS_API FTestCustomLayeredMove : public FLayeredMoveBase
+{
+	GENERATED_USTRUCT_BODY()
+
+	FTestCustomLayeredMove();
+	virtual ~FTestCustomLayeredMove() override {}
+	
+	// Velocity to apply to the actor. Could be additive or overriding depending on MixMode setting.
+	UPROPERTY(EditAnywhere, BlueprintReadWrite, Category = Mover, meta=(ForceUnits="cm/s"))
+	FVector LaunchVelocity = FVector::ZeroVector;
+
+	// Optional movement mode name to force the actor into before applying the impulse velocity.
+	UPROPERTY(EditAnywhere, BlueprintReadWrite, Category = Mover)
+	FName ForceMovementMode = NAME_None;
+
+	virtual void OnStart(const UMoverComponent* MoverComp, UMoverBlackboard* SimBlackboard) override;
+	
+	virtual void OnEnd(const UMoverComponent* MoverComp, UMoverBlackboard* SimBlackboard, float CurrentSimTimeMs) override;
+	
+	// Generate a movement 
+	virtual bool GenerateMove(const FMoverTickStartData& StartState, const FMoverTimeStep& TimeStep, const UMoverComponent* MoverComp, UMoverBlackboard* SimBlackboard, FProposedMove& OutProposedMove) override;
+
+	virtual FLayeredMoveBase* Clone() const override;
+
+	virtual void NetSerialize(FArchive& Ar) override;
+
+	virtual UScriptStruct* GetScriptStruct() const override;
+
+	virtual FString ToSimpleString() const override;
+
+	virtual void AddReferencedObjects(class FReferenceCollector& Collector) override;
+};
\ No newline at end of file
